Replace deprecated Google+ scope with basic userinfo scope

diff --git a/Engine/Plugins/Online/OnlineSubsystemGoogle/Source/Private/OnlineIdentityGoogleCommon.h b/Engine/Plugins/Online/OnlineSubsystemGoogle/Source/Private/OnlineIdentityGoogleCommon.h
--- a/Engine/Plugins/Online/OnlineSubsystemGoogle/Source/Private/OnlineIdentityGoogleCommon.h
+++ b/Engine/Plugins/Online/OnlineSubsystemGoogle/Source/Private/OnlineIdentityGoogleCommon.h
@@ -8,13 +8,11 @@
 #include "Interfaces/IHttpRequest.h"
 #include "OnlineSubsystemGooglePackage.h"
 
-// Google scope fields
+// Google scope fields: https://developers.google.com/identity/protocols/googlescopes
 // email profile
-// https://www.googleapis.com/auth/plus.login
-// https://www.googleapis.com/auth/plus.me 
 // https://www.googleapis.com/auth/userinfo.email
 // https://www.googleapis.com/auth/userinfo.profile
-#define GOOGLE_PERM_PUBLIC_PROFILE "https://www.googleapis.com/auth/plus.login"
+#define GOOGLE_PERM_PUBLIC_PROFILE "https://www.googleapis.com/auth/userinfo.profile"
 
 class FOnlineSubsystemGoogle;
 class FUserOnlineAccountGoogleCommon;
