Fixed a bug where Google OAUTH scopes were during logout and reloaded from INI, but code-only scope used during initialization was not being re-added.

diff --git a/Engine/Plugins/Online/OnlineSubsystemGoogle/Source/Private/Rest/OnlineIdentityGoogleRest.cpp b/Engine/Plugins/Online/OnlineSubsystemGoogle/Source/Private/Rest/OnlineIdentityGoogleRest.cpp
--- a/Engine/Plugins/Online/OnlineSubsystemGoogle/Source/Private/Rest/OnlineIdentityGoogleRest.cpp
+++ b/Engine/Plugins/Online/OnlineSubsystemGoogle/Source/Private/Rest/OnlineIdentityGoogleRest.cpp
@@ -420,6 +420,8 @@ bool FOnlineIdentityGoogle::Logout(int32 LocalUserNum)
 		UserIds.Remove(LocalUserNum);
 		// reset scope permissions
 		GConfig->GetArray(TEXT("OnlineSubsystemGoogle.OnlineIdentityGoogle"), TEXT("ScopeFields"), LoginURLDetails.ScopeFields, GEngineIni);
+		// always required login access fields
+		LoginURLDetails.ScopeFields.AddUnique(TEXT(GOOGLE_PERM_PUBLIC_PROFILE));
 
 		TriggerOnLoginFlowLogoutDelegates(LoginDomains);
 
