Remove references to UE4 in Google OSS

diff --git a/Engine/Plugins/Online/OnlineSubsystemGoogle/Source/OnlineSubsystemGoogle.Build.cs b/Engine/Plugins/Online/OnlineSubsystemGoogle/Source/OnlineSubsystemGoogle.Build.cs
--- a/Engine/Plugins/Online/OnlineSubsystemGoogle/Source/OnlineSubsystemGoogle.Build.cs
+++ b/Engine/Plugins/Online/OnlineSubsystemGoogle/Source/OnlineSubsystemGoogle.Build.cs
@@ -37,7 +37,6 @@ public class OnlineSubsystemGoogle : ModuleRules
 		if (Target.Platform == UnrealTargetPlatform.IOS)
 		{
 			PublicDefinitions.Add("WITH_GOOGLE=1");
-			PublicDefinitions.Add("UE4_GOOGLE_VER=4.0.1");
 			PrivateIncludePaths.Add("Private/IOS");
 
 			// These are iOS system libraries that Google depends on
