Fix PlatformAllowList for OnlineSubsystemGoogle plugin

diff --git a/Engine/Plugins/Online/OnlineSubsystemGoogle/OnlineSubsystemGoogle.uplugin b/Engine/Plugins/Online/OnlineSubsystemGoogle/OnlineSubsystemGoogle.uplugin
--- a/Engine/Plugins/Online/OnlineSubsystemGoogle/OnlineSubsystemGoogle.uplugin
+++ b/Engine/Plugins/Online/OnlineSubsystemGoogle/OnlineSubsystemGoogle.uplugin
@@ -10,17 +10,21 @@
 	"CreatedByURL" : "https://epicgames.com",
 	"EnabledByDefault" : false,
 
-  "Modules": [
-        {
-            "Name": "OnlineSubsystemGoogle",
-            "Type": "Runtime",
-            "LoadingPhase" : "Default",
-			      "PlatformAllowList" :
-			      [
-				      "Win64"
-			      ]
-        }
-    ],
+	"Modules": [
+		{
+			"Name": "OnlineSubsystemGoogle",
+			"Type": "Runtime",
+			"LoadingPhase" : "Default",
+			"PlatformAllowList" : [
+				"Android",
+				"IOS",
+				"Linux",
+				"LinuxArm64",
+				"Mac",
+				"Win64"
+			]
+		}
+	],
 	"Plugins": [
 		{
 			"Name": "OnlineSubsystem",
