Fix compilation issues on NDK r27c

diff --git a/Engine/Plugins/Online/OnlineSubsystemGoogle/Source/Private/Android/OnlineIdentityGoogle.cpp b/Engine/Plugins/Online/OnlineSubsystemGoogle/Source/Private/Android/OnlineIdentityGoogle.cpp
--- a/Engine/Plugins/Online/OnlineSubsystemGoogle/Source/Private/Android/OnlineIdentityGoogle.cpp
+++ b/Engine/Plugins/Online/OnlineSubsystemGoogle/Source/Private/Android/OnlineIdentityGoogle.cpp
@@ -211,7 +211,7 @@ bool FOnlineIdentityGoogle::Logout(int32 LocalUserNum)
 
 void FOnlineIdentityGoogle::OnLoginComplete(EGoogleLoginResponse InResponseCode, const TSharedPtr<FUserOnlineAccountGoogle>& User)
 {
-	UE_LOG_ONLINE_IDENTITY(Verbose, TEXT("OnLoginComplete %s %s"), ToString(InResponseCode));
+	UE_LOG_ONLINE_IDENTITY(Verbose, TEXT("OnLoginComplete %s"), ToString(InResponseCode));
 	ensure(LoginCompletionDelegate.IsBound());
 	LoginCompletionDelegate.ExecuteIfBound(InResponseCode, User);
 	LoginCompletionDelegate.Unbind();
