Fix compile failure due to recent iOS rotation code changes

#rnx

diff --git a/Engine/Plugins/Online/OnlineSubsystemGoogle/Source/Private/IOS/GoogleHelper.cpp b/Engine/Plugins/Online/OnlineSubsystemGoogle/Source/Private/IOS/GoogleHelper.cpp
--- a/Engine/Plugins/Online/OnlineSubsystemGoogle/Source/Private/IOS/GoogleHelper.cpp
+++ b/Engine/Plugins/Online/OnlineSubsystemGoogle/Source/Private/IOS/GoogleHelper.cpp
@@ -79,7 +79,7 @@ bool GetAuthTokenFromSignInResult(GIDGoogleUser* User, NSString* ServerAuthCode,
 
 - (void)ShowLoginUI: (NSArray*) InScopes
 {
-	[GIDSignIn.sharedInstance signInWithPresentingViewController:[IOSAppDelegate GetDelegate].IOSController
+	[GIDSignIn.sharedInstance signInWithPresentingViewController:(UIViewController*)[IOSAppDelegate GetDelegate].IOSController
 							hint:nil
 							additionalScopes:InScopes
 							completion:^(GIDSignInResult* SignInResult, NSError* Error)
@@ -120,7 +120,7 @@ bool GetAuthTokenFromSignInResult(GIDGoogleUser* User, NSString* ServerAuthCode,
 			else
 			{
 				[User addScopes: MissingScopes 
-					 presentingViewController: [IOSAppDelegate GetDelegate].IOSController
+					 presentingViewController: (UIViewController*)[IOSAppDelegate GetDelegate].IOSController
 					 completion: ^(GIDSignInResult* Result, NSError* Error)
 				{
 					[self DidSignIn: Result.user withServerAuthCode: Result.serverAuthCode withError: Error];
