- Deprecate OnlineJsonSerializer.h

#rnx

diff --git a/Engine/Plugins/Online/OnlineSubsystem/Source/Public/Interfaces/OnlinePurchaseInterface.h b/Engine/Plugins/Online/OnlineSubsystem/Source/Public/Interfaces/OnlinePurchaseInterface.h
--- a/Engine/Plugins/Online/OnlineSubsystem/Source/Public/Interfaces/OnlinePurchaseInterface.h
+++ b/Engine/Plugins/Online/OnlineSubsystem/Source/Public/Interfaces/OnlinePurchaseInterface.h
@@ -4,7 +4,7 @@
 
 #include "Interfaces/OnlineEntitlementsInterface.h"
 #include "Interfaces/OnlineStoreInterfaceV2.h"
-#include "OnlineJsonSerializer.h"
+#include "Serialization/JsonSerializerMacros.h"
 
 struct FOnlineError;
 
@@ -100,7 +100,7 @@ enum class EPurchaseTransactionState : uint8
  * Receipt result from checkout
  */
 class FPurchaseReceipt
-	: public FOnlineJsonSerializable
+	: public FJsonSerializable
 {
 public:
 	FPurchaseReceipt()
@@ -109,7 +109,7 @@ public:
 	}
 
 	struct FLineItemInfo
-		: public FOnlineJsonSerializable
+		: public FJsonSerializable
 	{
 		/** The platform identifier of this purchase type */
 		FString ItemName;
@@ -122,18 +122,18 @@ public:
 
 		inline bool IsRedeemable() const { return !ValidationInfo.IsEmpty(); }
 
-		BEGIN_ONLINE_JSON_SERIALIZER
-			ONLINE_JSON_SERIALIZE("itemName", ItemName);
-			ONLINE_JSON_SERIALIZE("uniqueItemId", UniqueId);
-			ONLINE_JSON_SERIALIZE("validationInfo", ValidationInfo);
-		END_ONLINE_JSON_SERIALIZER
+		BEGIN_JSON_SERIALIZER
+			JSON_SERIALIZE("itemName", ItemName);
+			JSON_SERIALIZE("uniqueItemId", UniqueId);
+			JSON_SERIALIZE("validationInfo", ValidationInfo);
+		END_JSON_SERIALIZER
 	};
 
 	/**
 	 * Single purchased offer offer
 	 */
 	struct FReceiptOfferEntry
-		: public FOnlineJsonSerializable
+		: public FJsonSerializable
 	{
 		FReceiptOfferEntry()
 			: Quantity(0)
@@ -154,12 +154,12 @@ public:
 		/** Information about the individual items purchased */
 		TArray<FLineItemInfo> LineItems;
 
-		BEGIN_ONLINE_JSON_SERIALIZER
-			ONLINE_JSON_SERIALIZE("namespace", Namespace);
-			ONLINE_JSON_SERIALIZE("offerId", OfferId);
-			ONLINE_JSON_SERIALIZE("quantity", Quantity);
-			ONLINE_JSON_SERIALIZE_ARRAY_SERIALIZABLE("items", LineItems, FLineItemInfo);
-		END_ONLINE_JSON_SERIALIZER
+		BEGIN_JSON_SERIALIZER
+			JSON_SERIALIZE("namespace", Namespace);
+			JSON_SERIALIZE("offerId", OfferId);
+			JSON_SERIALIZE("quantity", Quantity);
+			JSON_SERIALIZE_ARRAY_SERIALIZABLE("items", LineItems, FLineItemInfo);
+		END_JSON_SERIALIZER
 	};
 
 	/**
@@ -184,10 +184,10 @@ public:
 		ReceiptOffers.Add(MoveTemp(ReceiptOffer));
 	}
 
-	BEGIN_ONLINE_JSON_SERIALIZER
-		ONLINE_JSON_SERIALIZE("transactionId", TransactionId);
-		ONLINE_JSON_SERIALIZE_ARRAY_SERIALIZABLE("receiptList", ReceiptOffers, FReceiptOfferEntry);
-	END_ONLINE_JSON_SERIALIZER
+	BEGIN_JSON_SERIALIZER
+		JSON_SERIALIZE("transactionId", TransactionId);
+		JSON_SERIALIZE_ARRAY_SERIALIZABLE("receiptList", ReceiptOffers, FReceiptOfferEntry);
+	END_JSON_SERIALIZER
 
 public:
 	/** Unique Id for this transaction/order */
diff --git a/Engine/Plugins/Online/OnlineSubsystem/Source/Public/OnlineJsonSerializer.h b/Engine/Plugins/Online/OnlineSubsystem/Source/Public/OnlineJsonSerializer.h
--- a/Engine/Plugins/Online/OnlineSubsystem/Source/Public/OnlineJsonSerializer.h
+++ b/Engine/Plugins/Online/OnlineSubsystem/Source/Public/OnlineJsonSerializer.h
@@ -2,8 +2,11 @@
 
 #pragma once
 
+#include "Misc/CoreMiscDefines.h"
 #include "Serialization/JsonSerializerMacros.h"
 
+UE_DEPRECATED_HEADER(5.7, "OnlineJsonSerializer.h is deprecated, please use the regular JSON functionality in Serialization/JsonSerializerMacros.h");
+
 class FVariantData;
 
 #define BEGIN_ONLINE_JSON_SERIALIZER																BEGIN_JSON_SERIALIZER
diff --git a/Engine/Plugins/Online/OnlineSubsystemAmazon/Source/Private/OnlineIdentityAmazon.h b/Engine/Plugins/Online/OnlineSubsystemAmazon/Source/Private/OnlineIdentityAmazon.h
--- a/Engine/Plugins/Online/OnlineSubsystemAmazon/Source/Private/OnlineIdentityAmazon.h
+++ b/Engine/Plugins/Online/OnlineSubsystemAmazon/Source/Private/OnlineIdentityAmazon.h
@@ -3,7 +3,7 @@
 #pragma once
  
 #include "Interfaces/OnlineIdentityInterface.h"
-#include "OnlineJsonSerializer.h"
+#include "Serialization/JsonSerializerMacros.h"
 #include "OnlineSubsystemNames.h"
 
 class FOnlineSubsystemAmazon;
@@ -16,7 +16,7 @@ TEMP_UNIQUENETIDSTRING_SUBCLASS(FUniqueNetIdAmazon, AMAZON_SUBSYSTEM);
  */
 class FUserOnlineAccountAmazon : 
 	public FUserOnlineAccount,
-	public FOnlineJsonSerializable
+	public FJsonSerializable
 {
 public:
 
@@ -59,12 +59,12 @@ public:
 
 	// FJsonSerializable
 
-	BEGIN_ONLINE_JSON_SERIALIZER
-		ONLINE_JSON_SERIALIZE("gameAccountId", UserId);
-		ONLINE_JSON_SERIALIZE("internalToken", SecretKey);
-		ONLINE_JSON_SERIALIZE("authTicket", AuthTicket);
-		ONLINE_JSON_SERIALIZE_MAP("additionalAuthData", AdditionalAuthData);
-	END_ONLINE_JSON_SERIALIZER
+	BEGIN_JSON_SERIALIZER
+		JSON_SERIALIZE("gameAccountId", UserId);
+		JSON_SERIALIZE("internalToken", SecretKey);
+		JSON_SERIALIZE("authTicket", AuthTicket);
+		JSON_SERIALIZE_MAP("additionalAuthData", AdditionalAuthData);
+	END_JSON_SERIALIZER
 };
 /** Mapping from user id to internal online account info (only one per user) */
 typedef TMap<FString, TSharedRef<FUserOnlineAccountAmazon> > FUserOnlineAccountAmazonMap;
diff --git a/Engine/Plugins/Online/OnlineSubsystemFacebook/Source/Private/OnlineAccountFacebookCommon.h b/Engine/Plugins/Online/OnlineSubsystemFacebook/Source/Private/OnlineAccountFacebookCommon.h
--- a/Engine/Plugins/Online/OnlineSubsystemFacebook/Source/Private/OnlineAccountFacebookCommon.h
+++ b/Engine/Plugins/Online/OnlineSubsystemFacebook/Source/Private/OnlineAccountFacebookCommon.h
@@ -3,7 +3,7 @@
 #pragma once
  
 #include "OnlineSubsystemFacebookTypes.h"
-#include "OnlineJsonSerializer.h"
+#include "Serialization/JsonSerializerMacros.h"
 #include "OnlineSubsystemFacebookPackage.h"
 #include "Interfaces/IHttpRequest.h"
 
@@ -19,7 +19,7 @@
  */
 class FUserOnlineAccountFacebookCommon : 
 	public FUserOnlineAccount,
-	public FOnlineJsonSerializable
+	public FJsonSerializable
 {
 public:
 
@@ -120,14 +120,14 @@ private:
 
 	// FJsonSerializable
 
-	BEGIN_ONLINE_JSON_SERIALIZER
-		ONLINE_JSON_SERIALIZE(ME_FIELD_ID, UserId);
-		ONLINE_JSON_SERIALIZE(ME_FIELD_NAME, RealName);
-		ONLINE_JSON_SERIALIZE(ME_FIELD_FIRSTNAME, FirstName);
-		ONLINE_JSON_SERIALIZE(ME_FIELD_LASTNAME, LastName);
-		ONLINE_JSON_SERIALIZE_OBJECT_SERIALIZABLE(ME_FIELD_PICTURE, Picture);
-		ONLINE_JSON_SERIALIZE_SIMPLE_COPY(AccountData);
-	END_ONLINE_JSON_SERIALIZER
+	BEGIN_JSON_SERIALIZER
+		JSON_SERIALIZE(ME_FIELD_ID, UserId);
+		JSON_SERIALIZE(ME_FIELD_NAME, RealName);
+		JSON_SERIALIZE(ME_FIELD_FIRSTNAME, FirstName);
+		JSON_SERIALIZE(ME_FIELD_LASTNAME, LastName);
+		JSON_SERIALIZE_OBJECT_SERIALIZABLE(ME_FIELD_PICTURE, Picture);
+		JSON_SERIALIZE_SIMPLECOPY(AccountData);
+	END_JSON_SERIALIZER
 
 	/** Allow the FB identity to fill in our private members from it's callbacks */
 	friend class FOnlineIdentityFacebook;
diff --git a/Engine/Plugins/Online/OnlineSubsystemFacebook/Source/Private/OnlineFriendsFacebookCommon.h b/Engine/Plugins/Online/OnlineSubsystemFacebook/Source/Private/OnlineFriendsFacebookCommon.h
--- a/Engine/Plugins/Online/OnlineSubsystemFacebook/Source/Private/OnlineFriendsFacebookCommon.h
+++ b/Engine/Plugins/Online/OnlineSubsystemFacebook/Source/Private/OnlineFriendsFacebookCommon.h
@@ -6,7 +6,7 @@
 #include "Interfaces/OnlineFriendsInterface.h"
 #include "Interfaces/OnlinePresenceInterface.h"
 #include "OnlineSubsystemFacebookTypes.h"
-#include "OnlineJsonSerializer.h"
+#include "Serialization/JsonSerializerMacros.h"
 #include "Http.h"
 #include "OnlineSubsystemFacebookPackage.h"
 
@@ -97,11 +97,11 @@ private:
 	FOnlineUserPresence Presence;
 
 	// FJsonSerializable
-	BEGIN_ONLINE_JSON_SERIALIZER
-		ONLINE_JSON_SERIALIZE(FRIEND_FIELD_ID, UserIdStr);
-		ONLINE_JSON_SERIALIZE_OBJECT_SERIALIZABLE(FRIEND_FIELD_PICTURE, Picture);
-		ONLINE_JSON_SERIALIZE_SIMPLE_COPY(AccountData);
-	END_ONLINE_JSON_SERIALIZER
+	BEGIN_JSON_SERIALIZER
+		JSON_SERIALIZE(FRIEND_FIELD_ID, UserIdStr);
+		JSON_SERIALIZE_OBJECT_SERIALIZABLE(FRIEND_FIELD_PICTURE, Picture);
+		JSON_SERIALIZE_SIMPLECOPY(AccountData);
+	END_JSON_SERIALIZER
 
 	/** Allow the FB friends to fill in our private members from it's callbacks */
 	friend class FOnlineFriendsFacebookCommon;
diff --git a/Engine/Plugins/Online/OnlineSubsystemFacebook/Source/Private/OnlineSharingFacebookCommon.h b/Engine/Plugins/Online/OnlineSubsystemFacebook/Source/Private/OnlineSharingFacebookCommon.h
--- a/Engine/Plugins/Online/OnlineSubsystemFacebook/Source/Private/OnlineSharingFacebookCommon.h
+++ b/Engine/Plugins/Online/OnlineSubsystemFacebook/Source/Private/OnlineSharingFacebookCommon.h
@@ -7,7 +7,7 @@
 #include "Interfaces/OnlineSharingInterface.h"
 #include "Interfaces/OnlineIdentityInterface.h"
 #include "Interfaces/IHttpRequest.h"
-#include "OnlineJsonSerializer.h"
+#include "Serialization/JsonSerializerMacros.h"
 #include "OnlineSubsystemFacebookPackage.h"
 
 class FOnlineSubsystemFacebook;
@@ -41,11 +41,11 @@ private:
 	typedef TMap<EOnlineSharingCategory, TArray<FString>> FSharingPermissionsMap;
 	FSharingPermissionsMap SharingPermissionsMap;
 
-	class FFacebookPermissionsJson : public FOnlineJsonSerializable
+	class FFacebookPermissionsJson : public FJsonSerializable
 	{
 		public:
 			class FFacebookPermissionJson :
-				public FOnlineJsonSerializable
+				public FJsonSerializable
 			{
 
 			public:
@@ -57,19 +57,19 @@ private:
 				FFacebookPermissionJson() {}
 
 				// FJsonSerializable
-				BEGIN_ONLINE_JSON_SERIALIZER
-					ONLINE_JSON_SERIALIZE(PERM_JSON_PERMISSION_NAME, Name);
-					ONLINE_JSON_SERIALIZE(PERM_JSON_PERMISSION_STATUS, Status);
-				END_ONLINE_JSON_SERIALIZER
+				BEGIN_JSON_SERIALIZER
+					JSON_SERIALIZE(PERM_JSON_PERMISSION_NAME, Name);
+					JSON_SERIALIZE(PERM_JSON_PERMISSION_STATUS, Status);
+				END_JSON_SERIALIZER
 			};
 
 
 			/** Main error body */
 			TArray<FFacebookPermissionJson> Permissions;
 
-			BEGIN_ONLINE_JSON_SERIALIZER
-				ONLINE_JSON_SERIALIZE_ARRAY_SERIALIZABLE(PERM_JSON_PERMISSIONS, Permissions, FFacebookPermissionJson);
-			END_ONLINE_JSON_SERIALIZER
+			BEGIN_JSON_SERIALIZER
+				JSON_SERIALIZE_ARRAY_SERIALIZABLE(PERM_JSON_PERMISSIONS, Permissions, FFacebookPermissionJson);
+			END_JSON_SERIALIZER
 	};
 
 	EOnlineSharingCategory GetCategoryFromFacebookPermission(const FString& FacebookPermission) const;
diff --git a/Engine/Plugins/Online/OnlineSubsystemFacebook/Source/Private/OnlineSubsystemFacebookCommon.h b/Engine/Plugins/Online/OnlineSubsystemFacebook/Source/Private/OnlineSubsystemFacebookCommon.h
--- a/Engine/Plugins/Online/OnlineSubsystemFacebook/Source/Private/OnlineSubsystemFacebookCommon.h
+++ b/Engine/Plugins/Online/OnlineSubsystemFacebook/Source/Private/OnlineSubsystemFacebookCommon.h
@@ -5,7 +5,7 @@
 #include "OnlineSubsystem.h"
 #include "OnlineSubsystemImpl.h"
 #include "OnlineSubsystemFacebookPackage.h"
-#include "OnlineJsonSerializer.h"
+#include "Serialization/JsonSerializerMacros.h"
 
 /** Forward declarations of all interface classes */
 typedef TSharedPtr<class FOnlineIdentityFacebookCommon, ESPMode::ThreadSafe> FOnlineIdentityFacebookCommonPtr;
diff --git a/Engine/Plugins/Online/OnlineSubsystemFacebook/Source/Private/OnlineSubsystemFacebookTypes.h b/Engine/Plugins/Online/OnlineSubsystemFacebook/Source/Private/OnlineSubsystemFacebookTypes.h
--- a/Engine/Plugins/Online/OnlineSubsystemFacebook/Source/Private/OnlineSubsystemFacebookTypes.h
+++ b/Engine/Plugins/Online/OnlineSubsystemFacebook/Source/Private/OnlineSubsystemFacebookTypes.h
@@ -3,7 +3,7 @@
 #pragma once
 
 #include "OnlineSubsystemTypes.h"
-#include "OnlineJsonSerializer.h"
+#include "Serialization/JsonSerializerMacros.h"
 #include "OnlineSubsystemFacebookPackage.h"
 
 #define PICTURE_DATA "data"
@@ -115,7 +115,7 @@ private:
  * Facebook error from JSON payload
  */
 class FErrorFacebook :
-	public FOnlineJsonSerializable
+	public FJsonSerializable
 {
 public:
 
@@ -127,7 +127,7 @@ public:
 	}
 
 	class FErrorBody :
-		public FOnlineJsonSerializable
+		public FJsonSerializable
 	{
 
 	public:
@@ -145,13 +145,13 @@ public:
 		FErrorBody() {}
 
 		// FJsonSerializable
-		BEGIN_ONLINE_JSON_SERIALIZER
-			ONLINE_JSON_SERIALIZE("message", Message);
-			ONLINE_JSON_SERIALIZE("type", Type);
-			ONLINE_JSON_SERIALIZE("code", Code);
-			ONLINE_JSON_SERIALIZE("error_subcode", ErrorSubCode);
-			ONLINE_JSON_SERIALIZE("fbtrace_id", FBTraceId);
-		END_ONLINE_JSON_SERIALIZER
+		BEGIN_JSON_SERIALIZER
+			JSON_SERIALIZE("message", Message);
+			JSON_SERIALIZE("type", Type);
+			JSON_SERIALIZE("code", Code);
+			JSON_SERIALIZE("error_subcode", ErrorSubCode);
+			JSON_SERIALIZE("fbtrace_id", FBTraceId);
+		END_JSON_SERIALIZER
 	};
 
 	/** Main error body */
@@ -160,36 +160,36 @@ public:
 	/** @return debug output for logs */
 	FString ToDebugString() const { return FString::Printf(TEXT("%s [Type:%s Code:%d SubCode:%d Trace:%s]"), *Error.Message, *Error.Type, Error.Code, Error.ErrorSubCode, *Error.FBTraceId); }
 
-	BEGIN_ONLINE_JSON_SERIALIZER
-		ONLINE_JSON_SERIALIZE_OBJECT_SERIALIZABLE("error", Error);
-	END_ONLINE_JSON_SERIALIZER
+	BEGIN_JSON_SERIALIZER
+		JSON_SERIALIZE_OBJECT_SERIALIZABLE("error", Error);
+	END_JSON_SERIALIZER
 };
 
 /** Facebook profile picture */
 class FUserOnlineFacebookPicture :
-	public FOnlineJsonSerializable
+	public FJsonSerializable
 {
 public:
 
 	struct FPictureData :
-		public FOnlineJsonSerializable
+		public FJsonSerializable
 	{
 		/** Is this picture the default silhouette */
 		bool bIsSilhouette;
 		/** URL to picture content */
 		FString PictureURL;
 
-		BEGIN_ONLINE_JSON_SERIALIZER
-			ONLINE_JSON_SERIALIZE(PICTURE_DATA_SILHOUETTE, bIsSilhouette);
-			ONLINE_JSON_SERIALIZE(PICTURE_DATA_URL, PictureURL);
-		END_ONLINE_JSON_SERIALIZER
+		BEGIN_JSON_SERIALIZER
+			JSON_SERIALIZE(PICTURE_DATA_SILHOUETTE, bIsSilhouette);
+			JSON_SERIALIZE(PICTURE_DATA_URL, PictureURL);
+		END_JSON_SERIALIZER
 	};
 
 	/** User picture */
 	FPictureData PictureData;
 
 	// FJsonSerializable
-	BEGIN_ONLINE_JSON_SERIALIZER
-		ONLINE_JSON_SERIALIZE_OBJECT_SERIALIZABLE(PICTURE_DATA, PictureData);
-	END_ONLINE_JSON_SERIALIZER
+	BEGIN_JSON_SERIALIZER
+		JSON_SERIALIZE_OBJECT_SERIALIZABLE(PICTURE_DATA, PictureData);
+	END_JSON_SERIALIZER
 };
diff --git a/Engine/Plugins/Online/OnlineSubsystemGoogle/Source/Private/OnlineAccountGoogleCommon.h b/Engine/Plugins/Online/OnlineSubsystemGoogle/Source/Private/OnlineAccountGoogleCommon.h
--- a/Engine/Plugins/Online/OnlineSubsystemGoogle/Source/Private/OnlineAccountGoogleCommon.h
+++ b/Engine/Plugins/Online/OnlineSubsystemGoogle/Source/Private/OnlineAccountGoogleCommon.h
@@ -3,7 +3,7 @@
 #pragma once
  
 #include "OnlineSubsystemGoogleTypes.h"
-#include "OnlineJsonSerializer.h"
+#include "Serialization/JsonSerializerMacros.h"
 #include "OnlineSubsystemGooglePackage.h"
 #include "Interfaces/IHttpRequest.h"
 
@@ -12,7 +12,7 @@
  */
 class FUserOnlineAccountGoogleCommon : 
 	public FUserOnlineAccount,
-	public FOnlineJsonSerializable
+	public FJsonSerializable
 {
 public:
 
@@ -116,14 +116,14 @@ protected:
 private:
 
 	// FJsonSerializable
-	BEGIN_ONLINE_JSON_SERIALIZER
-		ONLINE_JSON_SERIALIZE("id", UserId); //v2
-		ONLINE_JSON_SERIALIZE("sub", UserId); //v3
-		ONLINE_JSON_SERIALIZE("given_name", FirstName);
-		ONLINE_JSON_SERIALIZE("family_name", LastName);
-		ONLINE_JSON_SERIALIZE("name", RealName);
-		ONLINE_JSON_SERIALIZE_SIMPLE_COPY(AccountData);
-	END_ONLINE_JSON_SERIALIZER
+	BEGIN_JSON_SERIALIZER
+		JSON_SERIALIZE("id", UserId); //v2
+		JSON_SERIALIZE("sub", UserId); //v3
+		JSON_SERIALIZE("given_name", FirstName);
+		JSON_SERIALIZE("family_name", LastName);
+		JSON_SERIALIZE("name", RealName);
+		JSON_SERIALIZE_SIMPLECOPY(AccountData);
+	END_JSON_SERIALIZER
 
 	/** Allow the Google identity to fill in our private members from it's callbacks */
 	friend class FOnlineIdentityGoogle;
diff --git a/Engine/Plugins/Online/OnlineSubsystemGoogle/Source/Private/OnlineIdentityGoogleCommon.h b/Engine/Plugins/Online/OnlineSubsystemGoogle/Source/Private/OnlineIdentityGoogleCommon.h
--- a/Engine/Plugins/Online/OnlineSubsystemGoogle/Source/Private/OnlineIdentityGoogleCommon.h
+++ b/Engine/Plugins/Online/OnlineSubsystemGoogle/Source/Private/OnlineIdentityGoogleCommon.h
@@ -4,7 +4,7 @@
  
 #include "Interfaces/OnlineIdentityInterface.h"
 #include "OnlineSubsystemGoogleTypes.h"
-#include "OnlineJsonSerializer.h"
+#include "Serialization/JsonSerializerMacros.h"
 #include "Interfaces/IHttpRequest.h"
 #include "OnlineSubsystemGooglePackage.h"
 
diff --git a/Engine/Plugins/Online/OnlineSubsystemGoogle/Source/Private/OnlineSubsystemGoogleTypes.h b/Engine/Plugins/Online/OnlineSubsystemGoogle/Source/Private/OnlineSubsystemGoogleTypes.h
--- a/Engine/Plugins/Online/OnlineSubsystemGoogle/Source/Private/OnlineSubsystemGoogleTypes.h
+++ b/Engine/Plugins/Online/OnlineSubsystemGoogle/Source/Private/OnlineSubsystemGoogleTypes.h
@@ -3,7 +3,7 @@
 #pragma once
 
 #include "OnlineSubsystemTypes.h"
-#include "OnlineJsonSerializer.h"
+#include "Serialization/JsonSerializerMacros.h"
 #include "OnlineSubsystemGooglePackage.h"
 
 class FAuthTokenGoogle;
@@ -43,7 +43,7 @@ public:
 
 private:
 
-	class FHeader : public FOnlineJsonSerializable
+	class FHeader : public FJsonSerializable
 	{
 
 	public:
@@ -51,13 +51,13 @@ private:
 		FString Alg;
 		/** Key Id */
 		FString Kid;
-		BEGIN_ONLINE_JSON_SERIALIZER
-			ONLINE_JSON_SERIALIZE("alg", Alg);
-			ONLINE_JSON_SERIALIZE("kid", Kid);
-		END_ONLINE_JSON_SERIALIZER
+		BEGIN_JSON_SERIALIZER
+			JSON_SERIALIZE("alg", Alg);
+			JSON_SERIALIZE("kid", Kid);
+		END_JSON_SERIALIZER
 	};
 
-	class FPayload : public FOnlineJsonSerializable
+	class FPayload : public FJsonSerializable
 	{
 
 	public:
@@ -89,22 +89,22 @@ private:
 		/** User locale */
 		FString Locale;
 
-		BEGIN_ONLINE_JSON_SERIALIZER
-			ONLINE_JSON_SERIALIZE("sub", Sub);
-			ONLINE_JSON_SERIALIZE("given_name", FirstName);
-			ONLINE_JSON_SERIALIZE("family_name", LastName);
-			ONLINE_JSON_SERIALIZE("name", RealName);
-			ONLINE_JSON_SERIALIZE("iss", ISS);
-			ONLINE_JSON_SERIALIZE("iat", IAT);
-			ONLINE_JSON_SERIALIZE("exp", EXP);
-			ONLINE_JSON_SERIALIZE("at_hash", ATHash);
-			ONLINE_JSON_SERIALIZE("aud", Aud);
-			ONLINE_JSON_SERIALIZE("email_verified", bEmailVerified);
-			ONLINE_JSON_SERIALIZE("azp", AZP);
-			ONLINE_JSON_SERIALIZE("email", Email);
-			ONLINE_JSON_SERIALIZE("picture", Picture);
-			ONLINE_JSON_SERIALIZE("locale", Locale);
-		END_ONLINE_JSON_SERIALIZER
+		BEGIN_JSON_SERIALIZER
+			JSON_SERIALIZE("sub", Sub);
+			JSON_SERIALIZE("given_name", FirstName);
+			JSON_SERIALIZE("family_name", LastName);
+			JSON_SERIALIZE("name", RealName);
+			JSON_SERIALIZE("iss", ISS);
+			JSON_SERIALIZE("iat", IAT);
+			JSON_SERIALIZE("exp", EXP);
+			JSON_SERIALIZE("at_hash", ATHash);
+			JSON_SERIALIZE("aud", Aud);
+			JSON_SERIALIZE("email_verified", bEmailVerified);
+			JSON_SERIALIZE("azp", AZP);
+			JSON_SERIALIZE("email", Email);
+			JSON_SERIALIZE("picture", Picture);
+			JSON_SERIALIZE("locale", Locale);
+		END_JSON_SERIALIZER
 	};
 
 	/** JWT header */
@@ -120,7 +120,7 @@ enum EGoogleRefreshToken : uint8 { GoogleRefreshToken };
  * Google auth token representation, both exchange and access tokens
  */
 class FAuthTokenGoogle :
-	public FOnlineJsonSerializable
+	public FJsonSerializable
 {
 public:
 
@@ -249,20 +249,20 @@ private:
 	/** Any addition auth data associated with the token */
 	FJsonSerializableKeyValueMap AuthData;
 
-	BEGIN_ONLINE_JSON_SERIALIZER
-		ONLINE_JSON_SERIALIZE("access_token", AccessToken);
-		ONLINE_JSON_SERIALIZE("token_type", TokenType);
-		ONLINE_JSON_SERIALIZE("expires_in", ExpiresIn);
-		ONLINE_JSON_SERIALIZE("refresh_token", RefreshToken);
-		ONLINE_JSON_SERIALIZE("id_token", IdToken);
-	END_ONLINE_JSON_SERIALIZER
+	BEGIN_JSON_SERIALIZER
+		JSON_SERIALIZE("access_token", AccessToken);
+		JSON_SERIALIZE("token_type", TokenType);
+		JSON_SERIALIZE("expires_in", ExpiresIn);
+		JSON_SERIALIZE("refresh_token", RefreshToken);
+		JSON_SERIALIZE("id_token", IdToken);
+	END_JSON_SERIALIZER
 };
 
 /**
  * Container for configuration info related to all Google API services
  */
 class FGoogleOpenIDConfiguration :
-	public FOnlineJsonSerializable
+	public FJsonSerializable
 {
 public:
 
@@ -313,21 +313,21 @@ public:
 
 private:
 
-	BEGIN_ONLINE_JSON_SERIALIZER
-		ONLINE_JSON_SERIALIZE("issuer", Issuer);
-		ONLINE_JSON_SERIALIZE("authorization_endpoint", AuthEndpoint);
-		ONLINE_JSON_SERIALIZE("token_endpoint", TokenEndpoint);
-		ONLINE_JSON_SERIALIZE("userinfo_endpoint", UserInfoEndpoint);
-		ONLINE_JSON_SERIALIZE("revocation_endpoint", RevocationEndpoint);
-		ONLINE_JSON_SERIALIZE("jwks_uri", JWKSURI);
-	END_ONLINE_JSON_SERIALIZER
+	BEGIN_JSON_SERIALIZER
+		JSON_SERIALIZE("issuer", Issuer);
+		JSON_SERIALIZE("authorization_endpoint", AuthEndpoint);
+		JSON_SERIALIZE("token_endpoint", TokenEndpoint);
+		JSON_SERIALIZE("userinfo_endpoint", UserInfoEndpoint);
+		JSON_SERIALIZE("revocation_endpoint", RevocationEndpoint);
+		JSON_SERIALIZE("jwks_uri", JWKSURI);
+	END_JSON_SERIALIZER
 };
 
 /**
  * Google error from JSON payload
  */
 class FErrorGoogle :
-	public FOnlineJsonSerializable
+	public FJsonSerializable
 {
 public:
 
@@ -345,8 +345,8 @@ public:
 
 	FString ToDebugString() const { return FString::Printf(TEXT("%s [Desc:%s]"), *Error, *Error_Description); }
 
-	BEGIN_ONLINE_JSON_SERIALIZER
-		ONLINE_JSON_SERIALIZE("error", Error);
-		ONLINE_JSON_SERIALIZE("error_description", Error_Description);
-	END_ONLINE_JSON_SERIALIZER
+	BEGIN_JSON_SERIALIZER
+		JSON_SERIALIZE("error", Error);
+		JSON_SERIALIZE("error_description", Error_Description);
+	END_JSON_SERIALIZER
 };
diff --git a/Engine/Plugins/Online/OnlineSubsystemGoogle/Source/Public/OnlineSubsystemGoogleCommon.h b/Engine/Plugins/Online/OnlineSubsystemGoogle/Source/Public/OnlineSubsystemGoogleCommon.h
--- a/Engine/Plugins/Online/OnlineSubsystemGoogle/Source/Public/OnlineSubsystemGoogleCommon.h
+++ b/Engine/Plugins/Online/OnlineSubsystemGoogle/Source/Public/OnlineSubsystemGoogleCommon.h
@@ -4,7 +4,7 @@
 
 #include "OnlineSubsystem.h"
 #include "OnlineSubsystemImpl.h"
-#include "OnlineJsonSerializer.h"
+#include "Serialization/JsonSerializerMacros.h"
 #include "OnlineSubsystemGooglePackage.h"
 
 /** Forward declarations of all interface classes */
diff --git a/Engine/Plugins/Online/OnlineSubsystemTencent/Source/Private/OnlineAsyncTasksTencent.h b/Engine/Plugins/Online/OnlineSubsystemTencent/Source/Private/OnlineAsyncTasksTencent.h
--- a/Engine/Plugins/Online/OnlineSubsystemTencent/Source/Private/OnlineAsyncTasksTencent.h
+++ b/Engine/Plugins/Online/OnlineSubsystemTencent/Source/Private/OnlineAsyncTasksTencent.h
@@ -12,7 +12,7 @@
 #include "Interfaces/OnlinePurchaseInterface.h"
 #include "OnlineSubsystemTencentTypes.h"
 #include "OnlineError.h"
-#include "OnlineJsonSerializer.h"
+#include "Serialization/JsonSerializerMacros.h"
 
 class FOnlineSubsystemTencent;
 class FUniqueNetIdRail;
@@ -756,7 +756,7 @@ protected:
 };
 
 struct FOnlineAsyncEventRailShowFloatingNotifyWindowPayload
-	: public FOnlineJsonSerializable
+	: public FJsonSerializable
 {
 public:
 	/** Title of dialog to display*/
@@ -767,11 +767,11 @@ public:
 	FString ButtonText;
 
 public:
-	BEGIN_ONLINE_JSON_SERIALIZER
-	ONLINE_JSON_SERIALIZE("title", DialogTitle);
-	ONLINE_JSON_SERIALIZE("desc", DialogText);
-	ONLINE_JSON_SERIALIZE("confirm_button", ButtonText);
-	END_ONLINE_JSON_SERIALIZER
+	BEGIN_JSON_SERIALIZER
+	JSON_SERIALIZE("title", DialogTitle);
+	JSON_SERIALIZE("desc", DialogText);
+	JSON_SERIALIZE("confirm_button", ButtonText);
+	END_JSON_SERIALIZER
 };
 
 class FOnlineAsyncEventRailShowFloatingNotifyWindow
